Significance Manager
Add experimental plugin that implements the framework to be extended for game specific needs
The significance manager provides a framework for scoring the importance of objects given one or more viewpoints. This significance score can then be used to make decisions about update frequency, whether to spawn effects, level of detail, etc.
Stat SignificanceManager will display performance numbers about the work the significance manager is doing
ShowDebug SignificanceManager will display on the HUD an ordered list of the managed objects. The number of objects displayed can be modified via the console variable SignificanceManager.ObjectsToShow and the list can filtered to only objects of a given tag via SignificanceManager.FilterTag

diff --git a/Engine/Plugins/Experimental/SignificanceManager/SignificanceManager.uplugin b/Engine/Plugins/Experimental/SignificanceManager/SignificanceManager.uplugin
new file mode 100644
--- /dev/null
+++ b/Engine/Plugins/Experimental/SignificanceManager/SignificanceManager.uplugin
@@ -0,0 +1,22 @@
+{
+	"FileVersion" : 3,
+	
+	"FriendlyName" : "Significance Manager",
+	"Version" : 1,
+	"VersionName" : "1.0",
+	"CreatedBy" : "Epic Games, Inc.",
+	"CreatedByURL" : "http://epicgames.com",
+	"EngineVersion" : "4.10.0",
+	"Description" : "The significance manager plugin provides an extensible framework for allowing games to calculate the significance of an object and change behavior in response.",
+	"Category" : "Programming",
+	"EnabledByDefault" : false,
+
+	"Modules" :
+	[
+		{
+			"Name" : "SignificanceManager",
+			"Type" : "Runtime",
+			"LoadingPhase" : "PreDefault"
+		},
+	]
+}
\ No newline at end of file
diff --git a/Engine/Plugins/Experimental/SignificanceManager/Source/SignificanceManager/Private/SignificanceManager.cpp b/Engine/Plugins/Experimental/SignificanceManager/Source/SignificanceManager/Private/SignificanceManager.cpp
new file mode 100644
--- /dev/null
+++ b/Engine/Plugins/Experimental/SignificanceManager/Source/SignificanceManager/Private/SignificanceManager.cpp
@@ -0,0 +1,292 @@
+// Copyright 1998-2015 Epic Games, Inc. All Rights Reserved.
+
+#include "SignificanceManagerPrivate.h"
+#include "SignificanceManager.h"
+
+IMPLEMENT_MODULE( FDefaultModuleImpl, SignificanceManager );
+
+DECLARE_CYCLE_STAT(TEXT("Update Total"), STAT_SignificanceManager_Update, STATGROUP_SignificanceManager);
+DECLARE_CYCLE_STAT(TEXT("Significance Update"), STAT_SignificanceManager_SignificanceUpdate, STATGROUP_SignificanceManager);
+DECLARE_CYCLE_STAT(TEXT("Significance Sort"), STAT_SignificanceManager_SignificanceSort, STATGROUP_SignificanceManager);
+DECLARE_CYCLE_STAT(TEXT("Unregister Object"), STAT_SignificanceManager_UnregisterObject, STATGROUP_SignificanceManager);
+DECLARE_CYCLE_STAT(TEXT("Significance Check"), STAT_SignificanceManager_SignificanceCheck, STATGROUP_SignificanceManager);
+DECLARE_CYCLE_STAT(TEXT("Register Object"), STAT_SignificanceManager_RegisterObject, STATGROUP_SignificanceManager);
+DECLARE_CYCLE_STAT(TEXT("Initial Significance Update"), STAT_SignificanceManager_InitialSignificanceUpdate, STATGROUP_SignificanceManager);
+DECLARE_DWORD_ACCUMULATOR_STAT(TEXT("Num Managed Objects"), STAT_SignificanceManager_NumObjects, STATGROUP_SignificanceManager);
+
+
+bool CompareBySignificance(const USignificanceManager::FManagedObjectInfo& A, const USignificanceManager::FManagedObjectInfo& B) 
+{ 
+	return A.GetSignificance() > B.GetSignificance(); 
+}
+
+USignificanceManager::USignificanceManager()
+	: Super()
+{
+	SignificanceManagerClassName = FStringClassReference(GetClass()); 
+
+	// Only put this in place for the actual USignificanceManager CDO
+	if (HasAnyFlags(RF_ClassDefaultObject) && GetClass() == USignificanceManager::StaticClass())
+	{
+		FWorldDelegates::OnPreWorldInitialization.AddUObject(this, &USignificanceManager::OnWorldInit);
+		FWorldDelegates::OnWorldCleanup.AddUObject(this, &USignificanceManager::OnWorldCleanup);
+		AHUD::OnShowDebugInfo.AddUObject(this, &USignificanceManager::OnShowDebugInfo);
+	}
+}
+
+USignificanceManager* USignificanceManager::Get(UWorld* World)
+{ 
+	return GetDefault<USignificanceManager>()->WorldSignificanceManagers.FindRef(World); 
+}
+
+void USignificanceManager::OnWorldInit(UWorld* World, const UWorld::InitializationValues IVS)
+{
+	if (World->IsGameWorld())
+	{
+		if (*SignificanceManagerClass == nullptr)
+		{
+			SignificanceManagerClass = LoadClass<USignificanceManager>(nullptr, *SignificanceManagerClassName.ToString());
+		}
+
+		if (*SignificanceManagerClass != nullptr)
+		{
+			WorldSignificanceManagers.Add(World, NewObject<USignificanceManager>(World, SignificanceManagerClass));
+		}
+	}
+}
+
+void USignificanceManager::OnWorldCleanup(UWorld* World, bool bSessionEnded, bool bCleanupResources)
+{
+	WorldSignificanceManagers.Remove(World);
+}
+
+void USignificanceManager::BeginDestroy()
+{
+	Super::BeginDestroy();
+
+	for (const TPair<UObject*, FManagedObjectInfo*>& ObjectToObjectInfoPair : ManagedObjects)
+	{
+		delete ObjectToObjectInfoPair.Value;
+	}
+	ManagedObjects.Reset();
+	ManagedObjectsByTag.Reset();
+}
+
+void USignificanceManager::RegisterObject(UObject* Object, const FName Tag, FSignificanceFunction SignificanceFunction)
+{
+	INC_DWORD_STAT(STAT_SignificanceManager_NumObjects);
+	SCOPE_CYCLE_COUNTER(STAT_SignificanceManager_RegisterObject);
+
+	check(Object);
+	checkf(!ManagedObjects.Contains(Object), TEXT("'%s' already added to significance manager. Original Tag: '%s' New Tag: '%s'"), *Object->GetName(), *ManagedObjects.FindChecked(Object)->GetTag().ToString(), *Tag.ToString());
+
+	FManagedObjectInfo* ObjectInfo = new FManagedObjectInfo(Object, Tag, SignificanceFunction);
+	
+	// Calculate initial significance
+	if (Viewpoints.Num())
+	{
+		SCOPE_CYCLE_COUNTER(STAT_SignificanceManager_InitialSignificanceUpdate);
+		ObjectInfo->UpdateSignificance(Viewpoints);
+	}
+
+	ManagedObjects.Add(Object, ObjectInfo);
+	TArray<const FManagedObjectInfo*>& ManagedObjectInfos = ManagedObjectsByTag.FindOrAdd(Tag);
+
+	if (ManagedObjectInfos.Num() > 0)
+	{
+		// Insert in to the sorted list
+		int32 LowIndex = 0;
+		int32 HighIndex = ManagedObjectInfos.Num() - 1;
+		while (true)
+		{
+			int32 MidIndex = LowIndex + (HighIndex - LowIndex) / 2;
+			if (CompareBySignificance(*ObjectInfo, *ManagedObjectInfos[MidIndex]))
+			{
+				if (LowIndex == MidIndex)
+				{
+					ManagedObjectInfos.Insert(ObjectInfo, LowIndex);
+					break;
+				}
+				else
+				{
+					HighIndex = MidIndex - 1;
+				}
+			}
+			else if (LowIndex == HighIndex)
+			{
+				ManagedObjectInfos.Insert(ObjectInfo, LowIndex + 1);
+				break;
+			}
+			else
+			{
+				LowIndex = MidIndex + 1;
+			}
+		}
+	}
+	else
+	{
+		ManagedObjectInfos.Add(ObjectInfo);
+	}
+}
+
+void USignificanceManager::UnregisterObject(UObject* Object)
+{
+	DEC_DWORD_STAT(STAT_SignificanceManager_NumObjects);
+	SCOPE_CYCLE_COUNTER(STAT_SignificanceManager_UnregisterObject);
+
+	FManagedObjectInfo* ObjectInfo;
+	if (ManagedObjects.RemoveAndCopyValue(Object, ObjectInfo))
+	{
+		TArray<const FManagedObjectInfo*>& ObjectsWithTag = ManagedObjectsByTag.FindChecked(ObjectInfo->GetTag());
+		if (ObjectsWithTag.Num() == 1)
+		{
+			check(ObjectsWithTag[0] == ObjectInfo);
+			ManagedObjectsByTag.Remove(ObjectInfo->GetTag());
+		}
+		else
+		{
+			ObjectsWithTag.RemoveSingle(ObjectInfo);
+		}
+		delete ObjectInfo;
+	}
+}
+
+const TArray<const USignificanceManager::FManagedObjectInfo*>& USignificanceManager::GetManagedObjects(const FName Tag) const
+{
+	if (const TArray<const FManagedObjectInfo*>* ObjectsWithTag = ManagedObjectsByTag.Find(Tag))
+	{
+		return *ObjectsWithTag;
+	}
+
+	static const TArray<const FManagedObjectInfo*> EmptySet;
+	return EmptySet;
+}
+
+void USignificanceManager::GetManagedObjects(TArray<const USignificanceManager::FManagedObjectInfo*>& OutManagedObjects, bool bInSignificanceOrder) const
+{
+	OutManagedObjects.Reserve(ManagedObjects.Num());
+	for (const TPair<FName, TArray<const FManagedObjectInfo*>>& TagToObjectInfoArrayPair : ManagedObjectsByTag)
+	{
+		OutManagedObjects.Append(TagToObjectInfoArrayPair.Value);
+	}
+	if (bInSignificanceOrder)
+	{
+		OutManagedObjects.Sort(&CompareBySignificance);
+	}
+}
+
+float USignificanceManager::GetSignificance(UObject* Object) const
+{
+	SCOPE_CYCLE_COUNTER(STAT_SignificanceManager_SignificanceCheck);
+	float Significance = 0.f;
+	if (FManagedObjectInfo* const* Info = ManagedObjects.Find(Object))
+	{
+		Significance = (*Info)->GetSignificance();
+	}
+	return Significance;
+}
+
+void USignificanceManager::FManagedObjectInfo::UpdateSignificance(const TArray<FTransform>& InViewpoints)
+{
+	if (InViewpoints.Num())
+	{
+		Significance = TNumericLimits<float>::Lowest();
+		for (const FTransform& Viewpoint : InViewpoints)
+		{
+			const float ViewpointSignificance = SignificanceFunction(Object, Viewpoint);
+			if (ViewpointSignificance > Significance)
+			{
+				Significance = ViewpointSignificance;
+			}
+		}
+	}
+	else
+	{
+		Significance = 0.f;
+	}
+}
+
+void USignificanceManager::Update(const TArray<FTransform>& InViewpoints)
+{
+	Viewpoints = InViewpoints;
+
+	SCOPE_CYCLE_COUNTER(STAT_SignificanceManager_Update);
+
+	{
+		SCOPE_CYCLE_COUNTER(STAT_SignificanceManager_SignificanceUpdate);
+		//TODO: Consider parallel for
+		for (TPair<UObject*, FManagedObjectInfo*>& ObjectToObjectInfoPair : ManagedObjects)
+		{
+			UObject* Object = ObjectToObjectInfoPair.Key;
+			FManagedObjectInfo* ObjectInfo = ObjectToObjectInfoPair.Value;
+
+			checkSlow(Object->IsValidLowLevel());
+
+			ObjectInfo->UpdateSignificance(Viewpoints);
+		}
+	}
+
+	{
+		SCOPE_CYCLE_COUNTER(STAT_SignificanceManager_SignificanceSort);
+		for (TPair<FName, TArray<const FManagedObjectInfo*>>& TagToObjectInfoArrayPair : ManagedObjectsByTag)
+		{
+			TagToObjectInfoArrayPair.Value.Sort(&CompareBySignificance);
+		}
+	}
+}
+
+static int32 GSignificanceManagerObjectsToShow = 15;
+static FAutoConsoleVariableRef CVarSignificanceManagerObjectsToShow(
+	TEXT("SignificanceManager.ObjectsToShow"),
+	GSignificanceManagerObjectsToShow,
+	TEXT("How many objects to display when ShowDebug SignificanceManager is enabled.\n"),
+	ECVF_Cheat
+	);
+
+static FAutoConsoleVariable CVarSignificanceManagerFilterTag(
+	TEXT("SignificanceManager.FilterTag"),
+	TEXT(""),
+	TEXT("Only display objects with the specified filter tag.  If None objects with any will be displayed.\n"),
+	ECVF_Cheat
+	);
+
+
+void USignificanceManager::OnShowDebugInfo(AHUD* HUD, UCanvas* Canvas, const FDebugDisplayInfo& DisplayInfo, float& YL, float& YPos)
+{
+	static const FName NAME_SignificanceManager("SignificanceManager");
+	if (Canvas && HUD->ShouldDisplayDebug(NAME_SignificanceManager))
+	{
+		if (HasAnyFlags(RF_ClassDefaultObject))
+		{
+			if (USignificanceManager* SignificanceManager = Get(HUD->GetWorld()))
+			{
+				SignificanceManager->OnShowDebugInfo(HUD, Canvas, DisplayInfo, YL, YPos);
+			}
+		}
+		else
+		{
+			Canvas->SetDrawColor(255,0,0);
+			UFont* RenderFont = GEngine->GetSmallFont();
+			YL = Canvas->DrawText(RenderFont, FString::Printf(TEXT("SIGNIFICANCE MANAGER - %d Managed Objects"), ManagedObjects.Num()), 4.0f, YPos );
+			YPos += YL;
+
+			const FName SignificanceManagerTag(*CVarSignificanceManagerFilterTag->GetString());
+			TArray<const FManagedObjectInfo*> AllObjects;
+			const TArray<const FManagedObjectInfo*>& ObjectsToShow = (SignificanceManagerTag.IsNone() ? AllObjects : GetManagedObjects(SignificanceManagerTag));
+			if (SignificanceManagerTag.IsNone())
+			{
+				GetManagedObjects(AllObjects, true);
+			}
+
+			const int32 NumObjectsToShow = FMath::Min(GSignificanceManagerObjectsToShow, ObjectsToShow.Num());
+			for (int32 Index = 0; Index < NumObjectsToShow; ++Index)
+			{
+				const FManagedObjectInfo* ObjectInfo = ObjectsToShow[Index];
+				const FString Str = FString::Printf(TEXT("%6.3f - %s (%s)"), ObjectInfo->GetSignificance(), *ObjectInfo->GetObject()->GetName(), *ObjectInfo->GetTag().ToString());
+				Canvas->SetDrawColor(FColor::White);
+				YL = Canvas->DrawText(RenderFont, Str, 4.0f, YPos);
+				YPos += YL;
+			}
+		}
+	}
+}
diff --git a/Engine/Plugins/Experimental/SignificanceManager/Source/SignificanceManager/Private/SignificanceManagerPrivate.h b/Engine/Plugins/Experimental/SignificanceManager/Source/SignificanceManager/Private/SignificanceManagerPrivate.h
new file mode 100644
--- /dev/null
+++ b/Engine/Plugins/Experimental/SignificanceManager/Source/SignificanceManager/Private/SignificanceManagerPrivate.h
@@ -0,0 +1,5 @@
+// Copyright 1998-2015 Epic Games, Inc. All Rights Reserved.
+
+
+#pragma once
+#include "Engine.h"
diff --git a/Engine/Plugins/Experimental/SignificanceManager/Source/SignificanceManager/Public/SignificanceManager.h b/Engine/Plugins/Experimental/SignificanceManager/Source/SignificanceManager/Public/SignificanceManager.h
new file mode 100644
--- /dev/null
+++ b/Engine/Plugins/Experimental/SignificanceManager/Source/SignificanceManager/Public/SignificanceManager.h
@@ -0,0 +1,129 @@
+// Copyright 1998-2015 Epic Games, Inc. All Rights Reserved.
+
+#pragma once
+
+#include "SignificanceManager.generated.h"
+
+DECLARE_STATS_GROUP(TEXT("Significance Manager"), STATGROUP_SignificanceManager, STATCAT_Advanced);
+
+/* The significance manager provides a framework for registering objects by tag to each have a significance
+ * value calculated from which a game specific subclass and game logic can make decisions about what level
+ * of detail objects should be at, tick frequency, whether to spawn effects, and other such functionality
+ *
+ * Each object that is registered must have a corresponding unregister event or else a dangling Object reference will
+ * be left resulting in an eventual crash once the Object has been garbage collected.
+ *
+ * Each user of the significance manager is expected to call the Update function from the appropriate location in the
+ * game code.  GameViewportClient::Tick may often serve as a good place to do this.
+ */
+UCLASS(config=Engine, defaultconfig)
+class SIGNIFICANCEMANAGER_API USignificanceManager : public UObject
+{
+	GENERATED_BODY()
+
+public:
+	typedef TFunction<float(UObject*,const FTransform&)> FSignificanceFunction;
+
+	struct FManagedObjectInfo
+	{
+		FManagedObjectInfo()
+			: Object(nullptr)
+			, Significance(0.f)
+		{
+		}
+
+		FManagedObjectInfo(UObject* InObject, FName InTag, FSignificanceFunction InSignificanceFunction)
+			: Object(InObject)
+			, Tag(InTag)
+			, Significance(0.f)
+			, SignificanceFunction(InSignificanceFunction)
+		{
+		}
+
+		UObject* GetObject() const { return Object; }
+		FName GetTag() const { return Tag; }
+		float GetSignificance() const { return Significance; }
+
+	private:
+		UObject* Object;
+		FName Tag;
+		float Significance;
+
+		FSignificanceFunction SignificanceFunction;
+
+		void UpdateSignificance(const TArray<FTransform>& ViewPoints);
+
+		// Allow SignificanceManager to call UpdateSignificance
+		friend USignificanceManager;
+	};
+
+	USignificanceManager();
+
+	// Begin UObject overrides
+	virtual void BeginDestroy() override;
+	// End UObject overrides
+
+	// Overridable function to update the managed objects' significance
+	virtual void Update(const TArray<FTransform>& Viewpoints);
+
+	// Overridable function used to register an object as managed by the significance manager
+	virtual void RegisterObject(UObject* Object, FName Tag, FSignificanceFunction SignificanceFunction);
+
+	// Overridable function used to unregister an object as managed by the significance manager
+	virtual void UnregisterObject(UObject* Object);
+
+	// Returns objects of specified tag, Tag must be specified or else an empty array will be returned
+	const TArray<const FManagedObjectInfo*>& GetManagedObjects(FName Tag) const;
+
+	// Returns all managed objects regardless of tag
+	void GetManagedObjects(TArray<const FManagedObjectInfo*>& OutManagedObjects, bool bInSignificanceOrder = false) const;
+
+	// Returns the significance value for a given object, returns 0 if object is not managed
+	float GetSignificance(UObject* Object) const;
+
+	// Returns the significance manager for the specified World
+	static USignificanceManager* Get(UWorld* World);
+
+	// Templated convenience function to return a significance manager cast to a known type
+	template<class T>
+	static T* Get(UWorld* World)
+	{
+		return CastChecked<T>(Get(World), ECastCheckedType::NullAllowed);
+	}
+
+protected:
+
+	const TArray<FTransform>& GetViewpoints() const { return Viewpoints; }
+
+private:
+
+	// The cached viewpoints for significance for calculating when a new object is registered
+	TArray<FTransform> Viewpoints;
+
+	// All objects being managed organized by Tag
+	TMap<FName, TArray<const FManagedObjectInfo*>> ManagedObjectsByTag;
+
+	// Reverse lookup map to find the tag for a given object
+	TMap<UObject*, FManagedObjectInfo*> ManagedObjects;
+
+	// Game specific significance class to instantiate
+	UPROPERTY(globalconfig, noclear, EditAnywhere, Category=DefaultClasses, meta=(MetaClass="SignificanceManager", DisplayName="Significance Manager Class"))
+	FStringClassReference SignificanceManagerClassName;
+
+	// Cached class for instantiating significance manager, only populated on CDO
+	UPROPERTY()
+	TSubclassOf<USignificanceManager>  SignificanceManagerClass;
+
+	// Map of worlds to their significance manager, only populated on CDO
+	UPROPERTY()
+	TMap<UWorld*, USignificanceManager*> WorldSignificanceManagers;
+
+	// Callback function registered with global world delegates to instantiate significance manager when a game world is created
+	void OnWorldInit(UWorld* World, const UWorld::InitializationValues IVS);
+
+	// Callback function registered with global world delegates to cleanup significance manager when a game world is destroyed
+	void OnWorldCleanup(UWorld* World, bool bSessionEnded, bool bCleanupResources);
+
+	// Callback function registered with HUD to supply debug info when ShowDebug SignificanceManager has been entered on the console
+	void OnShowDebugInfo(AHUD* HUD, UCanvas* Canvas, const FDebugDisplayInfo& DisplayInfo, float& YL, float& YPos);
+};
diff --git a/Engine/Plugins/Experimental/SignificanceManager/Source/SignificanceManager/SignificanceManager.Build.cs b/Engine/Plugins/Experimental/SignificanceManager/Source/SignificanceManager/SignificanceManager.Build.cs
new file mode 100644
--- /dev/null
+++ b/Engine/Plugins/Experimental/SignificanceManager/Source/SignificanceManager/SignificanceManager.Build.cs
@@ -0,0 +1,25 @@
+// Copyright 1998-2015 Epic Games, Inc. All Rights Reserved.
+
+using UnrealBuildTool;
+
+public class SignificanceManager : ModuleRules
+{
+	public SignificanceManager(TargetInfo Target)
+	{
+		PublicIncludePaths.Add("SignificanceManager/Public");
+
+		PrivateDependencyModuleNames.AddRange(
+			new string[] {
+				"Core",
+				"CoreUObject",
+		        "Engine",
+			}
+		);
+
+        PrivateIncludePaths.AddRange(
+            new string[] {
+                "SignificanceManager/Private"
+            })
+		;
+	}
+}
