Significance manager can now specify whether to create on the client and/or the server and can specify whether to sort ascending or descending

diff --git a/Engine/Plugins/Experimental/SignificanceManager/Source/SignificanceManager/Private/SignificanceManager.cpp b/Engine/Plugins/Experimental/SignificanceManager/Source/SignificanceManager/Private/SignificanceManager.cpp
--- a/Engine/Plugins/Experimental/SignificanceManager/Source/SignificanceManager/Private/SignificanceManager.cpp
+++ b/Engine/Plugins/Experimental/SignificanceManager/Source/SignificanceManager/Private/SignificanceManager.cpp
@@ -15,22 +15,41 @@ DECLARE_CYCLE_STAT(TEXT("Initial Significance Update"), STAT_SignificanceManager
 DECLARE_DWORD_ACCUMULATOR_STAT(TEXT("Num Managed Objects"), STAT_SignificanceManager_NumObjects, STATGROUP_SignificanceManager);
 
 
-bool CompareBySignificance(const USignificanceManager::FManagedObjectInfo& A, const USignificanceManager::FManagedObjectInfo& B) 
+bool CompareBySignificanceAscending(const USignificanceManager::FManagedObjectInfo& A, const USignificanceManager::FManagedObjectInfo& B) 
+{ 
+	return A.GetSignificance() < B.GetSignificance(); 
+}
+
+bool CompareBySignificanceDescending(const USignificanceManager::FManagedObjectInfo& A, const USignificanceManager::FManagedObjectInfo& B) 
 { 
 	return A.GetSignificance() > B.GetSignificance(); 
 }
 
+typedef bool CompareFunctionType(const USignificanceManager::FManagedObjectInfo&,const USignificanceManager::FManagedObjectInfo&);
+
+CompareFunctionType& PickCompareBySignificance(const bool bAscending)
+{
+	return (bAscending ? CompareBySignificanceAscending : CompareBySignificanceDescending);
+}
+
 USignificanceManager::USignificanceManager()
 	: Super()
 {
 	SignificanceManagerClassName = FStringClassReference(GetClass()); 
 
+	bCreateOnClient = true;
+	bCreateOnServer = true;
+	bSortSignificanceAscending = false;
+
 	// Only put this in place for the actual USignificanceManager CDO
 	if (HasAnyFlags(RF_ClassDefaultObject) && GetClass() == USignificanceManager::StaticClass())
 	{
 		FWorldDelegates::OnPreWorldInitialization.AddUObject(this, &USignificanceManager::OnWorldInit);
 		FWorldDelegates::OnWorldCleanup.AddUObject(this, &USignificanceManager::OnWorldCleanup);
-		AHUD::OnShowDebugInfo.AddUObject(this, &USignificanceManager::OnShowDebugInfo);
+		if (!IsRunningDedicatedServer())
+		{
+			AHUD::OnShowDebugInfo.AddUObject(this, &USignificanceManager::OnShowDebugInfo);
+		}
 	}
 }
 
@@ -50,7 +69,11 @@ void USignificanceManager::OnWorldInit(UWorld* World, const UWorld::Initializati
 
 		if (*SignificanceManagerClass != nullptr)
 		{
-			WorldSignificanceManagers.Add(World, NewObject<USignificanceManager>(World, SignificanceManagerClass));
+			USignificanceManager* ManagerToCreateDefault = SignificanceManagerClass->GetDefaultObject<USignificanceManager>();
+			if ((ManagerToCreateDefault->bCreateOnServer && !IsRunningClientOnly()) || (ManagerToCreateDefault->bCreateOnClient && !IsRunningDedicatedServer()))
+			{
+				WorldSignificanceManagers.Add(World, NewObject<USignificanceManager>(World, SignificanceManagerClass));
+			}
 		}
 	}
 }
@@ -97,10 +120,11 @@ void USignificanceManager::RegisterObject(UObject* Object, const FName Tag, FSig
 		// Insert in to the sorted list
 		int32 LowIndex = 0;
 		int32 HighIndex = ManagedObjectInfos.Num() - 1;
+		auto CompareFunction = PickCompareBySignificance(bSortSignificanceAscending);
 		while (true)
 		{
 			int32 MidIndex = LowIndex + (HighIndex - LowIndex) / 2;
-			if (CompareBySignificance(*ObjectInfo, *ManagedObjectInfos[MidIndex]))
+			if (CompareFunction(*ObjectInfo, *ManagedObjectInfos[MidIndex]))
 			{
 				if (LowIndex == MidIndex)
 				{
@@ -171,7 +195,7 @@ void USignificanceManager::GetManagedObjects(TArray<const USignificanceManager::
 	}
 	if (bInSignificanceOrder)
 	{
-		OutManagedObjects.Sort(&CompareBySignificance);
+		OutManagedObjects.Sort(PickCompareBySignificance(bSortSignificanceAscending));
 	}
 }
 
@@ -230,7 +254,7 @@ void USignificanceManager::Update(const TArray<FTransform>& InViewpoints)
 		SCOPE_CYCLE_COUNTER(STAT_SignificanceManager_SignificanceSort);
 		for (TPair<FName, TArray<const FManagedObjectInfo*>>& TagToObjectInfoArrayPair : ManagedObjectsByTag)
 		{
-			TagToObjectInfoArrayPair.Value.Sort(&CompareBySignificance);
+			TagToObjectInfoArrayPair.Value.Sort(PickCompareBySignificance(bSortSignificanceAscending));
 		}
 	}
 }
diff --git a/Engine/Plugins/Experimental/SignificanceManager/Source/SignificanceManager/Public/SignificanceManager.h b/Engine/Plugins/Experimental/SignificanceManager/Source/SignificanceManager/Public/SignificanceManager.h
--- a/Engine/Plugins/Experimental/SignificanceManager/Source/SignificanceManager/Public/SignificanceManager.h
+++ b/Engine/Plugins/Experimental/SignificanceManager/Source/SignificanceManager/Public/SignificanceManager.h
@@ -95,6 +95,15 @@ protected:
 
 	const TArray<FTransform>& GetViewpoints() const { return Viewpoints; }
 
+	// Whether the significance manager should be created on a client
+	uint32 bCreateOnClient:1;
+
+	// Whether the significance manager should be created on the server
+	uint32 bCreateOnServer:1;
+
+	// Whether the significance sort should sort high values to the end of the list
+	uint32 bSortSignificanceAscending:1;
+
 private:
 
 	// The cached viewpoints for significance for calculating when a new object is registered
