Signifiance manager should be a normal module not a game module.
Per: https://udn.unrealengine.com/questions/449189/why-does-fsignificancemanagermodule-extend-fdefaul.html
#rnx

diff --git a/Engine/Plugins/Runtime/SignificanceManager/Source/SignificanceManager/Public/SignificanceManager.h b/Engine/Plugins/Runtime/SignificanceManager/Source/SignificanceManager/Public/SignificanceManager.h
--- a/Engine/Plugins/Runtime/SignificanceManager/Source/SignificanceManager/Public/SignificanceManager.h
+++ b/Engine/Plugins/Runtime/SignificanceManager/Source/SignificanceManager/Public/SignificanceManager.h
@@ -22,7 +22,7 @@ struct FAutoCompleteCommand;
 DECLARE_STATS_GROUP(TEXT("Significance Manager"), STATGROUP_SignificanceManager, STATCAT_Advanced);
 
 /* Module definition for significance manager. Owns the references to created significance managers*/
-class SIGNIFICANCEMANAGER_API FSignificanceManagerModule : public FDefaultGameModuleImpl, public FGCObject
+class SIGNIFICANCEMANAGER_API FSignificanceManagerModule : public FDefaultModuleImpl, public FGCObject
 {
 public:
 	// Begin IModuleInterface overrides
