Make reregistering an object in the significance manager an ensure not a check
#rnx

diff --git a/Engine/Plugins/Runtime/SignificanceManager/Source/SignificanceManager/Private/SignificanceManager.cpp b/Engine/Plugins/Runtime/SignificanceManager/Source/SignificanceManager/Private/SignificanceManager.cpp
--- a/Engine/Plugins/Runtime/SignificanceManager/Source/SignificanceManager/Private/SignificanceManager.cpp
+++ b/Engine/Plugins/Runtime/SignificanceManager/Source/SignificanceManager/Private/SignificanceManager.cpp
@@ -22,6 +22,8 @@ DECLARE_CYCLE_STAT(TEXT("Register Object"), STAT_SignificanceManager_RegisterObj
 DECLARE_CYCLE_STAT(TEXT("Initial Significance Update"), STAT_SignificanceManager_InitialSignificanceUpdate, STATGROUP_SignificanceManager);
 DECLARE_DWORD_ACCUMULATOR_STAT(TEXT("Num Managed Objects"), STAT_SignificanceManager_NumObjects, STATGROUP_SignificanceManager);
 
+DEFINE_LOG_CATEGORY(LogSignificance);
+
 bool CompareBySignificanceAscending(const USignificanceManager::FManagedObjectInfo& A, const USignificanceManager::FManagedObjectInfo& B) 
 { 
 	return A.GetSignificance() < B.GetSignificance(); 
@@ -174,7 +176,12 @@ void USignificanceManager::RegisterManagedObject(FManagedObjectInfo* ObjectInfo)
 	UObject* Object = ObjectInfo->GetObject();
 
 	check(Object);
-	checkf(!ManagedObjects.Contains(Object), TEXT("'%s' already added to significance manager. Original Tag: '%s' New Tag: '%s'"), *Object->GetName(), *ManagedObjects.FindChecked(Object)->GetTag().ToString(), *ObjectInfo->GetTag().ToString());
+	if (!ensureMsgf(!ManagedObjects.Contains(Object), TEXT("'%s' already added to significance manager. Original Tag: '%s' New Tag: '%s'"), *Object->GetName(), *ManagedObjects.FindChecked(Object)->GetTag().ToString(), *ObjectInfo->GetTag().ToString()))
+	{
+		UE_LOG(LogSignificance, Warning, TEXT("'%s' already added to significance manager. Original Tag: '%s' New Tag: '%s'"), *Object->GetName(), *ManagedObjects.FindChecked(Object)->GetTag().ToString(), *ObjectInfo->GetTag().ToString());
+		delete ObjectInfo;
+		return;
+	}
 
 	if (ObjectInfo->GetPostSignificanceType() == EPostSignificanceType::Sequential)
 	{
diff --git a/Engine/Plugins/Runtime/SignificanceManager/Source/SignificanceManager/Public/SignificanceManager.h b/Engine/Plugins/Runtime/SignificanceManager/Source/SignificanceManager/Public/SignificanceManager.h
--- a/Engine/Plugins/Runtime/SignificanceManager/Source/SignificanceManager/Public/SignificanceManager.h
+++ b/Engine/Plugins/Runtime/SignificanceManager/Source/SignificanceManager/Public/SignificanceManager.h
@@ -20,6 +20,7 @@ class USignificanceManager;
 struct FAutoCompleteCommand;
 
 DECLARE_STATS_GROUP(TEXT("Significance Manager"), STATGROUP_SignificanceManager, STATCAT_Advanced);
+DECLARE_LOG_CATEGORY_EXTERN(LogSignificance, Log, All);
 
 /* Module definition for significance manager. Owns the references to created significance managers*/
 class SIGNIFICANCEMANAGER_API FSignificanceManagerModule : public FDefaultModuleImpl, public FGCObject
