- Add const param/return version of GetManagedObject to SignificanceManager

diff --git a/Engine/Plugins/Runtime/SignificanceManager/Source/SignificanceManager/Private/SignificanceManager.cpp b/Engine/Plugins/Runtime/SignificanceManager/Source/SignificanceManager/Private/SignificanceManager.cpp
--- a/Engine/Plugins/Runtime/SignificanceManager/Source/SignificanceManager/Private/SignificanceManager.cpp
+++ b/Engine/Plugins/Runtime/SignificanceManager/Source/SignificanceManager/Private/SignificanceManager.cpp
@@ -311,6 +311,16 @@ USignificanceManager::FManagedObjectInfo* USignificanceManager::GetManagedObject
 	return nullptr;
 }
 
+const USignificanceManager::USignificanceManager::FManagedObjectInfo* USignificanceManager::GetManagedObject(const UObject* Object) const
+{
+	if (FManagedObjectInfo* const* Info = ManagedObjects.Find(Object))
+	{
+		return (*Info);
+	}
+
+	return nullptr;
+}
+
 void USignificanceManager::GetManagedObjects(TArray<USignificanceManager::FManagedObjectInfo*>& OutManagedObjects, bool bInSignificanceOrder) const
 {
 	OutManagedObjects.Reserve(ManagedObjects.Num());
diff --git a/Engine/Plugins/Runtime/SignificanceManager/Source/SignificanceManager/Public/SignificanceManager.h b/Engine/Plugins/Runtime/SignificanceManager/Source/SignificanceManager/Public/SignificanceManager.h
--- a/Engine/Plugins/Runtime/SignificanceManager/Source/SignificanceManager/Public/SignificanceManager.h
+++ b/Engine/Plugins/Runtime/SignificanceManager/Source/SignificanceManager/Public/SignificanceManager.h
@@ -182,6 +182,9 @@ public:
 	// Returns the managed object for the passed-in object, if any. Otherwise returns nullptr
 	USignificanceManager::FManagedObjectInfo* GetManagedObject(UObject* Object) const;
 
+	// Returns the managed object for the passed-in object, if any. Otherwise returns nullptr
+	const USignificanceManager::FManagedObjectInfo* GetManagedObject(const UObject* Object) const;
+
 	// Returns the significance value for a given object, returns 0 if object is not managed
 	float GetSignificance(const UObject* Object) const;
 
