Allow OrderedBuckets for significance manager to be scaled during construction if you have some kind of scaling factor or setting in your game that needs to be taken into account for the setup.

diff --git a/Engine/Plugins/Runtime/SignificanceManager/Source/SignificanceManager/Public/OrderedBudget.h b/Engine/Plugins/Runtime/SignificanceManager/Source/SignificanceManager/Public/OrderedBudget.h
--- a/Engine/Plugins/Runtime/SignificanceManager/Source/SignificanceManager/Public/OrderedBudget.h
+++ b/Engine/Plugins/Runtime/SignificanceManager/Source/SignificanceManager/Public/OrderedBudget.h
@@ -31,9 +31,9 @@ public:
 	// Creating a table that contains:
 	//   1,1,2,2,2,3,3,3,3,3 (out of bounds = 4)
 	// Returns true if the budget was modified, or false if the existing budget already matched
-	bool RecreateBudget(const FString& Specification)
+	bool RecreateBudget(const FString& Specification, float InBudgetValuesScale = 1.0f)
 	{
-		const bool bSpecificationDiffers = Specification != BudgetString;
+		const bool bSpecificationDiffers = Specification != BudgetString || BudgetValuesScale != InBudgetValuesScale;
 
 		if (bSpecificationDiffers)
 		{
@@ -46,7 +46,7 @@ public:
 			int32 LevelIndex = 0;
 			for (const FString& BudgetItem : BudgetStrings)
 			{
-				const int32 CountForThisLevel = FCString::Atoi(*BudgetItem);
+				const int32 CountForThisLevel = FCString::Atoi(*BudgetItem) * InBudgetValuesScale;
 				check(CountForThisLevel >= 0);
 
 				for (int32 ThisLevelIndex = 0; ThisLevelIndex < CountForThisLevel; ++ThisLevelIndex)
@@ -57,6 +57,7 @@ public:
 				++LevelIndex;
 			}
 
+			BudgetValuesScale = InBudgetValuesScale;
 			ValueForOutOfBounds = LevelIndex;
 		}
 
@@ -67,6 +68,9 @@ private:
 	// This is the budget value for the i-th closest character (e.g., there will be be 10 entries if the budgets ranges sum to cover 10)
 	TArray<int32> BudgetValues;
 
+	// The scalar to apply to the budget values.
+	float BudgetValuesScale = 1.0f;
+
 	// This is the budget value for things with an index further away than BudgetByIndex.Num()
 	int32 ValueForOutOfBounds;
 
