Allow OrderedBuckets for significance manager to be scaled, does it by compressing the index, rather than generating a scaled, larger array.

diff --git a/Engine/Plugins/Runtime/SignificanceManager/Source/SignificanceManager/Public/OrderedBudget.h b/Engine/Plugins/Runtime/SignificanceManager/Source/SignificanceManager/Public/OrderedBudget.h
--- a/Engine/Plugins/Runtime/SignificanceManager/Source/SignificanceManager/Public/OrderedBudget.h
+++ b/Engine/Plugins/Runtime/SignificanceManager/Source/SignificanceManager/Public/OrderedBudget.h
@@ -14,6 +14,7 @@ public:
 	// Returns the budget number for the Index-th closest character
 	FORCEINLINE int32 GetBudgetForIndex(int32 Index) const
 	{
+		Index = Index / BudgetValuesScale;
 		return (Index < BudgetValues.Num()) ? BudgetValues[Index] : ValueForOutOfBounds;
 	}
 
@@ -31,9 +32,9 @@ public:
 	// Creating a table that contains:
 	//   1,1,2,2,2,3,3,3,3,3 (out of bounds = 4)
 	// Returns true if the budget was modified, or false if the existing budget already matched
-	bool RecreateBudget(const FString& Specification, float InBudgetValuesScale = 1.0f)
+	bool RecreateBudget(const FString& Specification)
 	{
-		const bool bSpecificationDiffers = Specification != BudgetString || BudgetValuesScale != InBudgetValuesScale;
+		const bool bSpecificationDiffers = Specification != BudgetString;
 
 		if (bSpecificationDiffers)
 		{
@@ -46,7 +47,7 @@ public:
 			int32 LevelIndex = 0;
 			for (const FString& BudgetItem : BudgetStrings)
 			{
-				const int32 CountForThisLevel = FCString::Atoi(*BudgetItem) * InBudgetValuesScale;
+				const int32 CountForThisLevel = FCString::Atoi(*BudgetItem);
 				check(CountForThisLevel >= 0);
 
 				for (int32 ThisLevelIndex = 0; ThisLevelIndex < CountForThisLevel; ++ThisLevelIndex)
@@ -57,18 +58,25 @@ public:
 				++LevelIndex;
 			}
 
-			BudgetValuesScale = InBudgetValuesScale;
 			ValueForOutOfBounds = LevelIndex;
 		}
 
 		return bSpecificationDiffers;
 	}
 
+	void SetBudgetScale(float InScale)
+	{
+		if (InScale > 0)
+		{
+			BudgetValuesScale = InScale;
+		}
+	}
+
 private:
 	// This is the budget value for the i-th closest character (e.g., there will be be 10 entries if the budgets ranges sum to cover 10)
 	TArray<int32> BudgetValues;
 
-	// The scalar to apply to the budget values.
+	// The scalar to apply to the budget values, it does so by compressing the index used when requesting the budget
 	float BudgetValuesScale = 1.0f;
 
 	// This is the budget value for things with an index further away than BudgetByIndex.Num()
