Remove deprecated function
#rnx

diff --git a/Engine/Plugins/Runtime/SignificanceManager/Source/SignificanceManager/Private/SignificanceManager.cpp b/Engine/Plugins/Runtime/SignificanceManager/Source/SignificanceManager/Private/SignificanceManager.cpp
--- a/Engine/Plugins/Runtime/SignificanceManager/Source/SignificanceManager/Private/SignificanceManager.cpp
+++ b/Engine/Plugins/Runtime/SignificanceManager/Source/SignificanceManager/Private/SignificanceManager.cpp
@@ -147,21 +147,6 @@ UWorld* USignificanceManager::GetWorld()const
 	return CastChecked<UWorld>(GetOuter());
 }
 
-void USignificanceManager::RegisterObject(UObject* Object, const FName Tag, FSignificanceFunction SignificanceFunction, USignificanceManager::EPostSignificanceType PostSignificanceType, FPostSignificanceFunction PostSignificanceFunction)
-{
-	FManagedObjectSignificanceFunction ManagedObjectSignificanceFunctionWrapper = [SignificanceFunction](const FManagedObjectInfo* ObjectInfo, const FTransform& Viewpoint)
-	{
-		return SignificanceFunction(ObjectInfo->GetObject(), Viewpoint);
-	};
-
-	FManagedObjectPostSignificanceFunction ManagedObjectPostSignficanceFunctionWrapper = [PostSignificanceFunction](const FManagedObjectInfo* ObjectInfo, const float OldSignificance, const float NewSignificance, const bool bFinal)
-	{
-		PostSignificanceFunction(ObjectInfo->GetObject(), OldSignificance, NewSignificance, bFinal);
-	};
-
-	USignificanceManager::RegisterObject(Object, Tag, ManagedObjectSignificanceFunctionWrapper, PostSignificanceType, ManagedObjectPostSignficanceFunctionWrapper);
-}
-
 void USignificanceManager::RegisterObject(UObject* Object, const FName Tag, FManagedObjectSignificanceFunction SignificanceFunction, USignificanceManager::EPostSignificanceType PostSignificanceType, FManagedObjectPostSignificanceFunction PostSignificanceFunction)
 {
 	FManagedObjectInfo* ObjectInfo = new FManagedObjectInfo(Object, Tag, SignificanceFunction, PostSignificanceType, PostSignificanceFunction);
diff --git a/Engine/Plugins/Runtime/SignificanceManager/Source/SignificanceManager/Public/SignificanceManager.h b/Engine/Plugins/Runtime/SignificanceManager/Source/SignificanceManager/Public/SignificanceManager.h
--- a/Engine/Plugins/Runtime/SignificanceManager/Source/SignificanceManager/Public/SignificanceManager.h
+++ b/Engine/Plugins/Runtime/SignificanceManager/Source/SignificanceManager/Public/SignificanceManager.h
@@ -160,10 +160,6 @@ public:
 	// Overridable function to update the managed objects' significance
 	virtual void Update(TArrayView<const FTransform> Viewpoints);
 
-	// Overridable function used to register an object as managed by the significance manager
-	UE_DEPRECATED(4.21, "Override RegisterObject that uses ManagedObject significance functions")
-	virtual void RegisterObject(UObject* Object, FName Tag, FSignificanceFunction SignificanceFunction, EPostSignificanceType InPostSignificanceType = EPostSignificanceType::None, FPostSignificanceFunction InPostSignificanceFunction = nullptr);
-
 	// Overridable function used to register an object as managed by the significance manager
 	virtual void RegisterObject(UObject* Object, FName Tag, FManagedObjectSignificanceFunction SignificanceFunction, EPostSignificanceType InPostSignificanceType = EPostSignificanceType::None, FManagedObjectPostSignificanceFunction InPostSignificanceFunction = nullptr);
 
