SignificanceManager - Maintain lists of objects to update rather than rebuilding each frame.

Previously we would build up the lists of objects to update, and the list of objects requring sequential post work each frame.
We now just update these lists in Register/Unregister to keep them in sync with ManagedObjects.
We do still make a copy of these to keep it thread safe during update, but overall it saves around ~20-40us in SignificanceManager::Update.

diff --git a/Engine/Plugins/Runtime/SignificanceManager/Source/SignificanceManager/Private/SignificanceManager.cpp b/Engine/Plugins/Runtime/SignificanceManager/Source/SignificanceManager/Private/SignificanceManager.cpp
--- a/Engine/Plugins/Runtime/SignificanceManager/Source/SignificanceManager/Private/SignificanceManager.cpp
+++ b/Engine/Plugins/Runtime/SignificanceManager/Source/SignificanceManager/Private/SignificanceManager.cpp
@@ -140,6 +140,8 @@ void USignificanceManager::BeginDestroy()
 	}
 	ManagedObjects.Reset();
 	ManagedObjectsByTag.Reset();
+	ObjArray.Reset();
+	ObjWithSequentialPostWork.Reset();
 }
 
 UWorld* USignificanceManager::GetWorld()const
@@ -170,7 +172,7 @@ void USignificanceManager::RegisterManagedObject(FManagedObjectInfo* ObjectInfo)
 
 	if (ObjectInfo->GetPostSignificanceType() == EPostSignificanceType::Sequential)
 	{
-		++ManagedObjectsWithSequentialPostWork;
+		ObjWithSequentialPostWork.Add({ ObjectInfo, ObjectInfo->GetSignificance() });
 	}
 
 	// Calculate initial significance
@@ -186,6 +188,7 @@ void USignificanceManager::RegisterManagedObject(FManagedObjectInfo* ObjectInfo)
 	}
 
 	ManagedObjects.Add(Object, ObjectInfo);
+	ObjArray.Add(ObjectInfo);
 	TArray<FManagedObjectInfo*>& ManagedObjectInfos = ManagedObjectsByTag.FindOrAdd(ObjectInfo->GetTag());
 
 	if (ManagedObjectInfos.Num() > 0)
@@ -236,9 +239,15 @@ void USignificanceManager::UnregisterObject(UObject* Object)
 	{
 		if (ObjectInfo->GetPostSignificanceType() == EPostSignificanceType::Sequential)
 		{
-			--ManagedObjectsWithSequentialPostWork;
+			const int32 Index = ObjWithSequentialPostWork.IndexOfByPredicate([ObjectInfo](const FSequentialPostWorkPair& WorkPair) { return WorkPair.ObjectInfo == ObjectInfo; });
+			if (Index != -1)
+			{
+				ObjWithSequentialPostWork.RemoveAtSwap(Index, 1, false);
+			}
 		}
 
+		ObjArray.RemoveSwap(ObjectInfo, false);
+
 		TArray<FManagedObjectInfo*>& ObjectsWithTag = ManagedObjectsByTag.FindChecked(ObjectInfo->GetTag());
 		if (ObjectsWithTag.Num() == 1)
 		{
@@ -267,8 +276,14 @@ void USignificanceManager::UnregisterAll(FName Tag)
 		{
 			if (ManagedObj->GetPostSignificanceType() == EPostSignificanceType::Sequential)
 			{
-				--ManagedObjectsWithSequentialPostWork;
+				const int32 Index = ObjWithSequentialPostWork.IndexOfByPredicate([ManagedObj](const FSequentialPostWorkPair& WorkPair) { return WorkPair.ObjectInfo == ManagedObj; });
+				if (Index != -1)
+				{
+					ObjWithSequentialPostWork.RemoveAtSwap(Index, 1, false);
+				}
 			}
+
+			ObjArray.RemoveSwap(ManagedObj, false);
 			ManagedObjects.Remove(ManagedObj->GetObject());
 			if (ManagedObj->PostSignificanceFunction != nullptr)
 			{
@@ -402,22 +417,12 @@ void USignificanceManager::Update(TArrayView<const FTransform> InViewpoints)
 	{
 		SCOPE_CYCLE_COUNTER(STAT_SignificanceManager_SignificanceUpdate);
 
-		check(ObjArray.Num() == 0 && ObjWithSequentialPostWork.Num() == 0);
-
-		ObjArray.Reserve(ManagedObjects.Num());
-		ObjWithSequentialPostWork.Reserve(ManagedObjectsWithSequentialPostWork);
+		// Copy the data we'll be working on this frame to avoid issues if objects are registered/unregistered while it runs.
+		// We don't need to clear them afterwards as we always overwrite them fully.
+		ObjArrayCopy = ObjArray;
+		ObjWithSequentialPostWorkCopy = ObjWithSequentialPostWork;
 
-		for (const TPair<UObject*, FManagedObjectInfo*>& ManagedObjectPair : ManagedObjects)
-		{
-			FManagedObjectInfo* ObjectInfo = ManagedObjectPair.Value;
-			ObjArray.Add(ObjectInfo);
-			if (ObjectInfo->GetPostSignificanceType() == EPostSignificanceType::Sequential)
-			{
-				ObjWithSequentialPostWork.Add({ ObjectInfo,ObjectInfo->GetSignificance() });
-			}
-		}
-
-		ParallelFor(ObjArray.Num(),
+		ParallelFor(ObjArrayCopy.Num(),
 			[&](int32 Index)
 		{
 			FManagedObjectInfo* ObjectInfo = ObjArray[Index];
@@ -427,14 +432,11 @@ void USignificanceManager::Update(TArrayView<const FTransform> InViewpoints)
 			ObjectInfo->UpdateSignificance(Viewpoints,bSortSignificanceAscending);
 		});
 
-		for (const FSequentialPostWorkPair& SequentialPostWorkPair : ObjWithSequentialPostWork)
+		for (const FSequentialPostWorkPair& SequentialPostWorkPair : ObjWithSequentialPostWorkCopy)
 		{
 			FManagedObjectInfo* ObjectInfo = SequentialPostWorkPair.ObjectInfo; 
 			ObjectInfo->PostSignificanceFunction(ObjectInfo, SequentialPostWorkPair.OldSignificance, ObjectInfo->GetSignificance(), false);
 		}
-
-		ObjArray.Reset();
-		ObjWithSequentialPostWork.Reset();
 	}
 
 	{
diff --git a/Engine/Plugins/Runtime/SignificanceManager/Source/SignificanceManager/Public/SignificanceManager.h b/Engine/Plugins/Runtime/SignificanceManager/Source/SignificanceManager/Public/SignificanceManager.h
--- a/Engine/Plugins/Runtime/SignificanceManager/Source/SignificanceManager/Public/SignificanceManager.h
+++ b/Engine/Plugins/Runtime/SignificanceManager/Source/SignificanceManager/Public/SignificanceManager.h
@@ -225,23 +225,26 @@ protected:
 
 private:
 
-	uint32 ManagedObjectsWithSequentialPostWork;
-
 	// All objects being managed organized by Tag
 	TMap<FName, TArray<FManagedObjectInfo*>> ManagedObjectsByTag;
 
 	// Reverse lookup map to find the tag for a given object
 	TMap<UObject*, FManagedObjectInfo*> ManagedObjects;
 
-	// Arrays used for ::Update. To avoid memory allocations, making them members
+	// Array of all managed objects that we use for iteration during update. This is kept in sync with the ManagedObjects map.
 	TArray<FManagedObjectInfo*> ObjArray;
+	// We copy ObjArray to this before running update to avoid mutations during the update. To avoid memory allocations, making it a member.
+	TArray<FManagedObjectInfo*> ObjArrayCopy;
 
 	struct FSequentialPostWorkPair
 	{
 		FManagedObjectInfo* ObjectInfo;
 		float OldSignificance;
 	};
+	// Array of all managed objects requiring sequential work that we use for iteration during update. This is kept in sync with the ManagedObjects map.
 	TArray<FSequentialPostWorkPair> ObjWithSequentialPostWork;
+	// We copy ObjWithSequentialPostWork to this before running update to avoid mutations during the update. To avoid memory allocations, making it a member.
+	TArray<FSequentialPostWorkPair> ObjWithSequentialPostWorkCopy;
 
 	// Game specific significance class to instantiate
 	UPROPERTY(globalconfig, noclear, EditAnywhere, Category=DefaultClasses, meta=(MetaClass="/Script/SignificanceManager.SignificanceManager", DisplayName="Significance Manager Class"))
