Back out  changelist 21273349 for now. It seems to be causing a significant perf regression. #jira none

diff --git a/Engine/Plugins/Runtime/SignificanceManager/Source/SignificanceManager/Private/SignificanceManager.cpp b/Engine/Plugins/Runtime/SignificanceManager/Source/SignificanceManager/Private/SignificanceManager.cpp
--- a/Engine/Plugins/Runtime/SignificanceManager/Source/SignificanceManager/Private/SignificanceManager.cpp
+++ b/Engine/Plugins/Runtime/SignificanceManager/Source/SignificanceManager/Private/SignificanceManager.cpp
@@ -140,8 +140,6 @@ void USignificanceManager::BeginDestroy()
 	}
 	ManagedObjects.Reset();
 	ManagedObjectsByTag.Reset();
-	ObjArray.Reset();
-	ObjWithSequentialPostWork.Reset();
 }
 
 UWorld* USignificanceManager::GetWorld()const
@@ -172,7 +170,7 @@ void USignificanceManager::RegisterManagedObject(FManagedObjectInfo* ObjectInfo)
 
 	if (ObjectInfo->GetPostSignificanceType() == EPostSignificanceType::Sequential)
 	{
-		ObjWithSequentialPostWork.Add({ ObjectInfo, ObjectInfo->GetSignificance() });
+		++ManagedObjectsWithSequentialPostWork;
 	}
 
 	// Calculate initial significance
@@ -188,7 +186,6 @@ void USignificanceManager::RegisterManagedObject(FManagedObjectInfo* ObjectInfo)
 	}
 
 	ManagedObjects.Add(Object, ObjectInfo);
-	ObjArray.Add(ObjectInfo);
 	TArray<FManagedObjectInfo*>& ManagedObjectInfos = ManagedObjectsByTag.FindOrAdd(ObjectInfo->GetTag());
 
 	if (ManagedObjectInfos.Num() > 0)
@@ -250,15 +247,9 @@ void USignificanceManager::UnregisterObject(UObject* Object)
 	{
 		if (ObjectInfo->GetPostSignificanceType() == EPostSignificanceType::Sequential)
 		{
-			const int32 Index = ObjWithSequentialPostWork.IndexOfByPredicate([ObjectInfo](const FSequentialPostWorkPair& WorkPair) { return WorkPair.ObjectInfo == ObjectInfo; });
-			if (Index != -1)
-			{
-				ObjWithSequentialPostWork.RemoveAtSwap(Index, 1, false);
-			}
+			--ManagedObjectsWithSequentialPostWork;
 		}
 
-		ObjArray.RemoveSwap(ObjectInfo, false);
-
 		TArray<FManagedObjectInfo*>& ObjectsWithTag = ManagedObjectsByTag.FindChecked(ObjectInfo->GetTag());
 		if (ObjectsWithTag.Num() == 1)
 		{
@@ -287,14 +278,8 @@ void USignificanceManager::UnregisterAll(FName Tag)
 		{
 			if (ManagedObj->GetPostSignificanceType() == EPostSignificanceType::Sequential)
 			{
-				const int32 Index = ObjWithSequentialPostWork.IndexOfByPredicate([ManagedObj](const FSequentialPostWorkPair& WorkPair) { return WorkPair.ObjectInfo == ManagedObj; });
-				if (Index != -1)
-				{
-					ObjWithSequentialPostWork.RemoveAtSwap(Index, 1, false);
-				}
+				--ManagedObjectsWithSequentialPostWork;
 			}
-
-			ObjArray.RemoveSwap(ManagedObj, false);
 			ManagedObjects.Remove(ManagedObj->GetObject());
 			if (ManagedObj->PostSignificanceFunction != nullptr)
 			{
@@ -428,12 +413,22 @@ void USignificanceManager::Update(TArrayView<const FTransform> InViewpoints)
 	{
 		SCOPE_CYCLE_COUNTER(STAT_SignificanceManager_SignificanceUpdate);
 
-		// Copy the data we'll be working on this frame to avoid issues if objects are registered/unregistered while it runs.
-		// We don't need to clear them afterwards as we always overwrite them fully.
-		ObjArrayCopy = ObjArray;
-		ObjWithSequentialPostWorkCopy = ObjWithSequentialPostWork;
+		check(ObjArray.Num() == 0 && ObjWithSequentialPostWork.Num() == 0);
+
+		ObjArray.Reserve(ManagedObjects.Num());
+		ObjWithSequentialPostWork.Reserve(ManagedObjectsWithSequentialPostWork);
 
-		ParallelFor(ObjArrayCopy.Num(),
+		for (const TPair<UObject*, FManagedObjectInfo*>& ManagedObjectPair : ManagedObjects)
+		{
+			FManagedObjectInfo* ObjectInfo = ManagedObjectPair.Value;
+			ObjArray.Add(ObjectInfo);
+			if (ObjectInfo->GetPostSignificanceType() == EPostSignificanceType::Sequential)
+			{
+				ObjWithSequentialPostWork.Add({ ObjectInfo,ObjectInfo->GetSignificance() });
+			}
+		}
+
+		ParallelFor(ObjArray.Num(),
 			[&](int32 Index)
 		{
 			FManagedObjectInfo* ObjectInfo = ObjArray[Index];
@@ -443,11 +438,14 @@ void USignificanceManager::Update(TArrayView<const FTransform> InViewpoints)
 			ObjectInfo->UpdateSignificance(Viewpoints,bSortSignificanceAscending);
 		});
 
-		for (const FSequentialPostWorkPair& SequentialPostWorkPair : ObjWithSequentialPostWorkCopy)
+		for (const FSequentialPostWorkPair& SequentialPostWorkPair : ObjWithSequentialPostWork)
 		{
 			FManagedObjectInfo* ObjectInfo = SequentialPostWorkPair.ObjectInfo; 
 			ObjectInfo->PostSignificanceFunction(ObjectInfo, SequentialPostWorkPair.OldSignificance, ObjectInfo->GetSignificance(), false);
 		}
+
+		ObjArray.Reset();
+		ObjWithSequentialPostWork.Reset();
 	}
 
 	{
diff --git a/Engine/Plugins/Runtime/SignificanceManager/Source/SignificanceManager/Public/SignificanceManager.h b/Engine/Plugins/Runtime/SignificanceManager/Source/SignificanceManager/Public/SignificanceManager.h
--- a/Engine/Plugins/Runtime/SignificanceManager/Source/SignificanceManager/Public/SignificanceManager.h
+++ b/Engine/Plugins/Runtime/SignificanceManager/Source/SignificanceManager/Public/SignificanceManager.h
@@ -227,26 +227,23 @@ protected:
 
 private:
 
+	uint32 ManagedObjectsWithSequentialPostWork;
+
 	// All objects being managed organized by Tag
 	TMap<FName, TArray<FManagedObjectInfo*>> ManagedObjectsByTag;
 
 	// Reverse lookup map to find the tag for a given object
 	TMap<UObject*, FManagedObjectInfo*> ManagedObjects;
 
-	// Array of all managed objects that we use for iteration during update. This is kept in sync with the ManagedObjects map.
+	// Arrays used for ::Update. To avoid memory allocations, making them members
 	TArray<FManagedObjectInfo*> ObjArray;
-	// We copy ObjArray to this before running update to avoid mutations during the update. To avoid memory allocations, making it a member.
-	TArray<FManagedObjectInfo*> ObjArrayCopy;
 
 	struct FSequentialPostWorkPair
 	{
 		FManagedObjectInfo* ObjectInfo;
 		float OldSignificance;
 	};
-	// Array of all managed objects requiring sequential work that we use for iteration during update. This is kept in sync with the ManagedObjects map.
 	TArray<FSequentialPostWorkPair> ObjWithSequentialPostWork;
-	// We copy ObjWithSequentialPostWork to this before running update to avoid mutations during the update. To avoid memory allocations, making it a member.
-	TArray<FSequentialPostWorkPair> ObjWithSequentialPostWorkCopy;
 
 	// Game specific significance class to instantiate
 	UPROPERTY(globalconfig, noclear, EditAnywhere, Category=DefaultClasses, meta=(MetaClass="/Script/SignificanceManager.SignificanceManager", DisplayName="Significance Manager Class"))
