Significance Manager: Don't look for multiple matches in cached object array, as its elements are unique

diff --git a/Engine/Plugins/Runtime/SignificanceManager/Source/SignificanceManager/Private/SignificanceManager.cpp b/Engine/Plugins/Runtime/SignificanceManager/Source/SignificanceManager/Private/SignificanceManager.cpp
--- a/Engine/Plugins/Runtime/SignificanceManager/Source/SignificanceManager/Private/SignificanceManager.cpp
+++ b/Engine/Plugins/Runtime/SignificanceManager/Source/SignificanceManager/Private/SignificanceManager.cpp
@@ -307,7 +307,7 @@ void USignificanceManager::UnregisterObject(UObject* Object)
 			}
 		}
 
-		ObjArray.RemoveSwap(ObjectInfo, EAllowShrinking::No);
+		ObjArray.RemoveSingleSwap(ObjectInfo, EAllowShrinking::No);
 
 		TArray<FManagedObjectInfo*>& ObjectsWithTag = ManagedObjectsByTag.FindChecked(ObjectInfo->GetTag());
 		if (ObjectsWithTag.Num() == 1)
@@ -344,7 +344,7 @@ void USignificanceManager::UnregisterAll(FName Tag)
 				}
 			}
 
-			ObjArray.RemoveSwap(ManagedObj, EAllowShrinking::No);
+			ObjArray.RemoveSingleSwap(ManagedObj, EAllowShrinking::No);
 			ManagedObjects.Remove(ManagedObj->GetObject());
 			if (ManagedObj->PostSignificanceFunction != nullptr)
 			{
