PR #12649: Virtual USignificanceManager::OnShowDebugInfo

diff --git a/Engine/Plugins/Runtime/SignificanceManager/Source/SignificanceManager/Public/SignificanceManager.h b/Engine/Plugins/Runtime/SignificanceManager/Source/SignificanceManager/Public/SignificanceManager.h
--- a/Engine/Plugins/Runtime/SignificanceManager/Source/SignificanceManager/Public/SignificanceManager.h
+++ b/Engine/Plugins/Runtime/SignificanceManager/Source/SignificanceManager/Public/SignificanceManager.h
@@ -161,6 +161,9 @@ protected:
 	// Internal function that takes the managed object info and registers it with the significance manager
 	void RegisterManagedObject(FManagedObjectInfo* ObjectInfo);
 
+	// Callback function registered with HUD to supply debug info when ShowDebug SignificanceManager has been entered on the console
+	virtual void OnShowDebugInfo(AHUD* HUD, UCanvas* Canvas, const FDebugDisplayInfo& DisplayInfo, float& YL, float& YPos);
+
 	// Whether the significance manager should be created on a client. Only used from CDO and 
 	uint32 bCreateOnClient:1;
 
@@ -200,9 +203,6 @@ private:
 	UPROPERTY(globalconfig, noclear, EditAnywhere, Category=DefaultClasses, meta=(MetaClass="/Script/SignificanceManager.SignificanceManager", DisplayName="Significance Manager Class"))
 	FSoftClassPath SignificanceManagerClassName;
 
-	// Callback function registered with HUD to supply debug info when ShowDebug SignificanceManager has been entered on the console
-	void OnShowDebugInfo(AHUD* HUD, UCanvas* Canvas, const FDebugDisplayInfo& DisplayInfo, float& YL, float& YPos);
-
 	// Friend FSignificanceManagerModule so it can call OnShowDebugInfo and check 
 	friend class FSignificanceManagerModule;
 };
