Replace some usages of FORCEINLINE with inline in SignificanceManager.

#rnx

diff --git a/Engine/Plugins/Runtime/SignificanceManager/Source/SignificanceManager/Public/OrderedBudget.h b/Engine/Plugins/Runtime/SignificanceManager/Source/SignificanceManager/Public/OrderedBudget.h
--- a/Engine/Plugins/Runtime/SignificanceManager/Source/SignificanceManager/Public/OrderedBudget.h
+++ b/Engine/Plugins/Runtime/SignificanceManager/Source/SignificanceManager/Public/OrderedBudget.h
@@ -14,14 +14,14 @@ public:
 	}
 
 	// Returns the budget number for the Index-th closest character
-	FORCEINLINE int32 GetBudgetForIndex(int32 Index) const
+	int32 GetBudgetForIndex(int32 Index) const
 	{
 		Index = Index / BudgetValuesScale;
 		return (Index < BudgetValues.Num()) ? BudgetValues[Index] : ValueForOutOfBounds;
 	}
 
 	// Returns the number of items specified in the budget
-	FORCEINLINE int32 GetBudgetLength() const
+	int32 GetBudgetLength() const
 	{
 		return BudgetValues.Num();
 	}
diff --git a/Engine/Plugins/Runtime/SignificanceManager/Source/SignificanceManager/Public/SignificanceManager.h b/Engine/Plugins/Runtime/SignificanceManager/Source/SignificanceManager/Public/SignificanceManager.h
--- a/Engine/Plugins/Runtime/SignificanceManager/Source/SignificanceManager/Public/SignificanceManager.h
+++ b/Engine/Plugins/Runtime/SignificanceManager/Source/SignificanceManager/Public/SignificanceManager.h
@@ -85,11 +85,11 @@ public:
 			}
 		}
 
-		FORCEINLINE UObject* GetObject() const { return Object; }
-		FORCEINLINE FName GetTag() const { return Tag; }
-		FORCEINLINE float GetSignificance() const { return Significance; }
+		UObject* GetObject() const { return Object; }
+		FName GetTag() const { return Tag; }
+		float GetSignificance() const { return Significance; }
 		FManagedObjectSignificanceFunction GetSignificanceFunction() const { return SignificanceFunction; }
-		FORCEINLINE EPostSignificanceType GetPostSignificanceType() const { return PostSignificanceType; }
+		EPostSignificanceType GetPostSignificanceType() const { return PostSignificanceType; }
 		FManagedObjectPostSignificanceFunction GetPostSignificanceNotifyDelegate() const { return PostSignificanceFunction; }
 
 	private:
@@ -149,7 +149,7 @@ public:
 
 	// Templated convenience function to return a significance manager cast to a known type
 	template<class T>
-	FORCEINLINE static T* Get(const UWorld* World)
+	static T* Get(const UWorld* World)
 	{
 		return CastChecked<T>(Get(World), ECastCheckedType::NullAllowed);
 	}
